F - Static analysis fixes
#rnx

diff --git a/Engine/Source/Developer/GraphColor/Private/graphColorVertices.cpp b/Engine/Source/Developer/GraphColor/Private/graphColorVertices.cpp
--- a/Engine/Source/Developer/GraphColor/Private/graphColorVertices.cpp
+++ b/Engine/Source/Developer/GraphColor/Private/graphColorVertices.cpp
@@ -376,6 +376,7 @@ int _AssignColorToVertex(ColorVerticesContext *context, graphP theGraph, int v)
     	return NOTOK;
 
     // Run the neighbor list of v and unflag all the colors in use
+	CA_ASSUME(theGraph);
     e = gp_GetFirstArc(theGraph, v);
     while (gp_IsArc(e))
     {
diff --git a/Engine/Source/Developer/GraphColor/Private/graphUtils.cpp b/Engine/Source/Developer/GraphColor/Private/graphUtils.cpp
--- a/Engine/Source/Developer/GraphColor/Private/graphUtils.cpp
+++ b/Engine/Source/Developer/GraphColor/Private/graphUtils.cpp
@@ -636,7 +636,8 @@ void gp_AttachArc(graphP theGraph, int v, int e, int link, int newArc)
     	 gp_SetAdjacentArc(theGraph, newArc, 1^link, NIL);
 
     	 // newArcs's elink is e2
-    	 gp_SetAdjacentArc(theGraph, newArc, link, e2);
+		 CA_ASSUME(theGraph);
+		 gp_SetAdjacentArc(theGraph, newArc, link, e2);
 
     	 // if e2 is an arc, then e2's 1^link is newArc, else v's 1^link is newArc
     	 if (gp_IsArc(e2))
@@ -1069,8 +1070,9 @@ int _IdentifyVertices(graphP theGraph, int u, int v, int eBefore)
         	gp_SetLastArc(theGraph, u, gp_GetLastArc(theGraph, v));
         }
 
-        gp_SetFirstArc(theGraph, v, NIL);
-        gp_SetLastArc(theGraph, v, NIL);
+		CA_ASSUME(theGraph);
+		gp_SetFirstArc(theGraph, v, NIL);
+		gp_SetLastArc(theGraph, v, NIL);
     }
 
     return OK;
diff --git a/Engine/Source/Runtime/Renderer/Private/TranslucentLighting.cpp b/Engine/Source/Runtime/Renderer/Private/TranslucentLighting.cpp
--- a/Engine/Source/Runtime/Renderer/Private/TranslucentLighting.cpp
+++ b/Engine/Source/Runtime/Renderer/Private/TranslucentLighting.cpp
@@ -1604,6 +1604,7 @@ void FDeferredShadingSceneRenderer::InjectTranslucentVolumeLightingArray(FRHICom
 	for (int32 ViewIndex = 0; ViewIndex < Views.Num(); ViewIndex++)
 	{
 		// non-shadowed, non-light function lights
+		CA_SUPPRESS(6385);
 		InjectTranslucentLightArray(RHICmdList, Views[ViewIndex], LightInjectionData[ViewIndex], ViewIndex);
 	}
 
diff --git a/Engine/Source/Runtime/Windows/D3D11RHI/Private/D3D11Query.cpp b/Engine/Source/Runtime/Windows/D3D11RHI/Private/D3D11Query.cpp
--- a/Engine/Source/Runtime/Windows/D3D11RHI/Private/D3D11Query.cpp
+++ b/Engine/Source/Runtime/Windows/D3D11RHI/Private/D3D11Query.cpp
@@ -616,8 +616,10 @@ void FD3D11BufferedGPUTiming::InitDynamicRHI()
 			QueryDesc.Query = D3D11_QUERY_TIMESTAMP;
 			QueryDesc.MiscFlags = 0;
 
+			CA_SUPPRESS(6385);	// Doesn't like COM
 			D3DResult = D3DRHI->GetDevice()->CreateQuery(&QueryDesc,StartTimestamps[TimestampIndex].GetInitReference());
 			GIsSupported = GIsSupported && (D3DResult == S_OK);
+			CA_SUPPRESS(6385);	// Doesn't like COM
 			D3DResult = D3DRHI->GetDevice()->CreateQuery(&QueryDesc,EndTimestamps[TimestampIndex].GetInitReference());
 			GIsSupported = GIsSupported && (D3DResult == S_OK);
 		}
