F - Attempt to fix static analysis again
#rnx

diff --git a/Engine/Source/Developer/GraphColor/Private/graphColorVertices.cpp b/Engine/Source/Developer/GraphColor/Private/graphColorVertices.cpp
--- a/Engine/Source/Developer/GraphColor/Private/graphColorVertices.cpp
+++ b/Engine/Source/Developer/GraphColor/Private/graphColorVertices.cpp
@@ -349,6 +349,9 @@ int _AssignColorToVertex(ColorVerticesContext *context, graphP theGraph, int v)
 {
 	int e, w, color;
 
+	CA_ASSUME(theGraph);
+	CA_ASSUME(theGraph->E);
+
 	// Run the neighbor list of v and flag all the colors in use
     e = gp_GetFirstArc(theGraph, v);
     while (gp_IsArc(e))
@@ -376,7 +379,6 @@ int _AssignColorToVertex(ColorVerticesContext *context, graphP theGraph, int v)
     	return NOTOK;
 
     // Run the neighbor list of v and unflag all the colors in use
-	CA_ASSUME(theGraph);
     e = gp_GetFirstArc(theGraph, v);
     while (gp_IsArc(e))
     {
diff --git a/Engine/Source/Developer/GraphColor/Private/graphUtils.cpp b/Engine/Source/Developer/GraphColor/Private/graphUtils.cpp
--- a/Engine/Source/Developer/GraphColor/Private/graphUtils.cpp
+++ b/Engine/Source/Developer/GraphColor/Private/graphUtils.cpp
@@ -610,7 +610,9 @@ int  e, degree;
 
 void gp_AttachArc(graphP theGraph, int v, int e, int link, int newArc)
 {
-     if (gp_IsArc(e))
+	CA_ASSUME(theGraph);
+	CA_ASSUME(theGraph->E);
+	if (gp_IsArc(e))
      {
     	 int e2 = gp_GetAdjacentArc(theGraph, e, link);
 
@@ -636,8 +638,7 @@ void gp_AttachArc(graphP theGraph, int v, int e, int link, int newArc)
     	 gp_SetAdjacentArc(theGraph, newArc, 1^link, NIL);
 
     	 // newArcs's elink is e2
-		 CA_ASSUME(theGraph);
-		 gp_SetAdjacentArc(theGraph, newArc, link, e2);
+    	 gp_SetAdjacentArc(theGraph, newArc, link, e2);
 
     	 // if e2 is an arc, then e2's 1^link is newArc, else v's 1^link is newArc
     	 if (gp_IsArc(e2))
@@ -1070,9 +1071,8 @@ int _IdentifyVertices(graphP theGraph, int u, int v, int eBefore)
         	gp_SetLastArc(theGraph, u, gp_GetLastArc(theGraph, v));
         }
 
-		CA_ASSUME(theGraph);
-		gp_SetFirstArc(theGraph, v, NIL);
-		gp_SetLastArc(theGraph, v, NIL);
+        gp_SetFirstArc(theGraph, v, NIL);
+        gp_SetLastArc(theGraph, v, NIL);
     }
 
     return OK;
