Fixed static analysis warnings due false positive warnings

diff --git a/Engine/Source/Developer/GraphColor/Private/graphColorVertices.cpp b/Engine/Source/Developer/GraphColor/Private/graphColorVertices.cpp
--- a/Engine/Source/Developer/GraphColor/Private/graphColorVertices.cpp
+++ b/Engine/Source/Developer/GraphColor/Private/graphColorVertices.cpp
@@ -92,7 +92,12 @@ int _GetContractibleNeighbors(ColorVerticesContext *context, int v, int *pu, int
 
 int gp_ColorVertices(graphP theGraph)
 {
-    ColorVerticesContext *context = NULL;
+	// Avoid false positive static analysis C6011 warning. Remove when it is fixed
+	CA_ASSUME(theGraph);
+	CA_ASSUME(theGraph->E);
+	CA_ASSUME(theGraph->V);
+
+	ColorVerticesContext *context = NULL;
     int v, deg;
     int u=0, w=0, contractible;
 
@@ -276,6 +281,11 @@ int _IsConstantTimeContractible(ColorVerticesContext *context, int v)
 
 int _GetContractibleNeighbors(ColorVerticesContext *context, int v, int *pu, int *pw)
 {
+	// Avoid false positive static analysis C6011 warning. Remove when it is fixed
+	CA_ASSUME(theGraph);
+	CA_ASSUME(theGraph->E);
+	CA_ASSUME(theGraph->V);
+
 	int lowDegreeNeighbors[5], i, j, n=0, e;
 	graphP theGraph = context->theGraph;
 
@@ -347,10 +357,12 @@ int _GetVertexToReduce(ColorVerticesContext *context, graphP theGraph)
 
 int _AssignColorToVertex(ColorVerticesContext *context, graphP theGraph, int v)
 {
-	int e, w, color;
-
+	// Avoid false positive static analysis C6011 warning. Remove when it is fixed
 	CA_ASSUME(theGraph);
 	CA_ASSUME(theGraph->E);
+	CA_ASSUME(theGraph->V);
+
+	int e, w, color;
 
 	// Run the neighbor list of v and flag all the colors in use
     e = gp_GetFirstArc(theGraph, v);
@@ -405,6 +417,11 @@ int gp_GetNumColorsUsed(graphP theGraph)
 
 void gp_CopyColors(graphP theGraph, unsigned char* colors)
 {
+	// Avoid false positive static analysis C6011 warning. Remove when it is fixed
+	CA_ASSUME(theGraph);
+	CA_ASSUME(theGraph->E);
+	CA_ASSUME(theGraph->V);
+
 	ColorVerticesContext *context = (ColorVerticesContext *) gp_GetExtension(theGraph, COLORVERTICES_ID);
 	
 	for ( int v = gp_GetFirstVertex(theGraph); gp_VertexInRange(theGraph, v); v++)
diff --git a/Engine/Source/Developer/GraphColor/Private/graphUtils.cpp b/Engine/Source/Developer/GraphColor/Private/graphUtils.cpp
--- a/Engine/Source/Developer/GraphColor/Private/graphUtils.cpp
+++ b/Engine/Source/Developer/GraphColor/Private/graphUtils.cpp
@@ -133,7 +133,10 @@ graphP theGraph = (graphP) malloc(sizeof(baseGraphStructure));
 
 void _InitFunctionTable(graphP theGraph)
 {
-     theGraph->functions.fpInitGraph = _InitGraph;
+	// Avoid false positive static analysis C6011 warning. Remove when it is fixed
+	CA_ASSUME(theGraph);
+
+	theGraph->functions.fpInitGraph = _InitGraph;
 
      theGraph->functions.fpHideEdge = _HideEdge;
      theGraph->functions.fpRestoreVertex = _RestoreVertex;
@@ -180,7 +183,10 @@ int gp_InitGraph(graphP theGraph, int N)
 
 int  _InitGraph(graphP theGraph, int N)
 {
-	 int  Vsize, VIsize, Esize, stackSize;
+	// Avoid false positive static analysis C6011 warning. Remove when it is fixed
+	CA_ASSUME(theGraph);
+
+	int  Vsize, VIsize, Esize, stackSize;
 
 	 // Compute the vertex and edge capacities of the graph
      theGraph->N = N;
@@ -217,6 +223,11 @@ int  _InitGraph(graphP theGraph, int N)
  ********************************************************************/
 void _InitVertices(graphP theGraph)
 {
+	// Avoid false positive static analysis C6011 warning. Remove when it is fixed
+	CA_ASSUME(theGraph);
+	CA_ASSUME(theGraph->E);
+	CA_ASSUME(theGraph->V);
+
 #if NIL == 0
 	memset(theGraph->V, NIL_CHAR, gp_VertexIndexBound(theGraph) * sizeof(vertexRec));
 #elif NIL == -1
@@ -236,6 +247,11 @@ void _InitVertices(graphP theGraph)
  ********************************************************************/
 void _InitEdges(graphP theGraph)
 {
+	// Avoid false positive static analysis C6011 warning. Remove when it is fixed
+	CA_ASSUME(theGraph);
+	CA_ASSUME(theGraph->E);
+	CA_ASSUME(theGraph->V);
+
 #if NIL == 0
 	memset(theGraph->E, NIL_CHAR, gp_EdgeIndexBound(theGraph) * sizeof(edgeRec));
 #elif NIL == -1
@@ -258,6 +274,11 @@ void _InitEdges(graphP theGraph)
  ********************************************************************/
 int gp_GetArcCapacity(graphP theGraph)
 {
+	// Avoid false positive static analysis C6011 warning. Remove when it is fixed
+	CA_ASSUME(theGraph);
+	CA_ASSUME(theGraph->E);
+	CA_ASSUME(theGraph->V);
+
 	return theGraph->arcCapacity - gp_GetFirstEdge(theGraph);
 }
 
@@ -317,6 +338,11 @@ int gp_EnsureArcCapacity(graphP theGraph, int requiredArcCapacity)
 	if (theGraph == NULL || requiredArcCapacity <= 0)
 		return NOTOK;
 
+	// Avoid false positive static analysis C6011 warning. Remove when it is fixed
+	CA_ASSUME(theGraph);
+	CA_ASSUME(theGraph->E);
+	CA_ASSUME(theGraph->V);
+
 	// Train callers to only ask for an even number of arcs, since
 	// two are required per edge or directed edge.
 	if (requiredArcCapacity & 1)
@@ -343,7 +369,12 @@ int gp_EnsureArcCapacity(graphP theGraph, int requiredArcCapacity)
 
 void _ClearVisitedFlags(graphP theGraph)
 {
-	 _ClearVertexVisitedFlags(theGraph, TRUE);
+	// Avoid false positive static analysis C6011 warning. Remove when it is fixed
+	CA_ASSUME(theGraph);
+	CA_ASSUME(theGraph->E);
+	CA_ASSUME(theGraph->V);
+
+	_ClearVertexVisitedFlags(theGraph, TRUE);
 }
 
 /********************************************************************
@@ -352,6 +383,11 @@ void _ClearVisitedFlags(graphP theGraph)
 
 void _ClearVertexVisitedFlags(graphP theGraph, int includeVirtualVertices)
 {
+	// Avoid false positive static analysis C6011 warning. Remove when it is fixed
+	CA_ASSUME(theGraph);
+	CA_ASSUME(theGraph->E);
+	CA_ASSUME(theGraph->V);
+
 	int  v;
 
 	for (v = gp_GetFirstVertex(theGraph); gp_VertexInRange(theGraph, v); v++)
@@ -370,7 +406,10 @@ void _ClearVertexVisitedFlags(graphP theGraph, int includeVirtualVertices)
 
 void _ClearGraph(graphP theGraph)
 {
-     if (theGraph->V != NULL)
+	// Avoid false positive static analysis C6011 warning. Remove when it is fixed
+	CA_ASSUME(theGraph);
+
+	if (theGraph->V != NULL)
      {
           free(theGraph->V);
           theGraph->V = NULL;
@@ -435,6 +474,14 @@ int  v, e, Esize;
          return NOTOK;
      }
 
+	 // Avoid false positive static analysis C6011 warning. Remove when it is fixed
+	 CA_ASSUME(srcGraph);
+	 CA_ASSUME(dstGraph);
+	 CA_ASSUME(srcGraph->V);
+	 CA_ASSUME(dstGraph->V);
+	 CA_ASSUME(srcGraph->N);
+	 CA_ASSUME(dstGraph->N);
+
      // Ensure dstGraph has the required arc capacity; this expands
      // dstGraph if needed, but does not contract.  An error is only
      // returned if the expansion fails.
@@ -506,7 +553,12 @@ int  v, e, Esize;
 
 int  gp_IsNeighbor(graphP theGraph, int u, int v)
 {
-int  e = gp_GetFirstArc(theGraph, u);
+	// Avoid false positive static analysis C6011 warning. Remove when it is fixed
+	CA_ASSUME(theGraph);
+	CA_ASSUME(theGraph->E);
+	CA_ASSUME(theGraph->V);
+
+	int  e = gp_GetFirstArc(theGraph, u);
 
      while (gp_IsArc(e))
      {
@@ -537,7 +589,12 @@ int  e = gp_GetFirstArc(theGraph, u);
 
 int  gp_GetNeighborEdgeRecord(graphP theGraph, int u, int v)
 {
-int  e;
+	// Avoid false positive static analysis C6011 warning. Remove when it is fixed
+	CA_ASSUME(theGraph);
+	CA_ASSUME(theGraph->E);
+	CA_ASSUME(theGraph->V);
+
+	int  e;
 
      if (gp_IsNotVertex(u) || gp_IsNotVertex(v))
     	 return NIL + NOTOK - NOTOK;
@@ -571,7 +628,12 @@ int  e;
 
 int  gp_GetVertexDegree(graphP theGraph, int v)
 {
-int  e, degree;
+	// Avoid false positive static analysis C6011 warning. Remove when it is fixed
+	CA_ASSUME(theGraph);
+	CA_ASSUME(theGraph->E);
+	CA_ASSUME(theGraph->V);
+	
+	int  e, degree;
 
      if (theGraph==NULL || gp_IsNotVertex(v))
     	 return 0 + NOTOK - NOTOK;
@@ -610,8 +672,11 @@ int  e, degree;
 
 void gp_AttachArc(graphP theGraph, int v, int e, int link, int newArc)
 {
+	// Avoid false positive static analysis C6011 warning. Remove when it is fixed
 	CA_ASSUME(theGraph);
 	CA_ASSUME(theGraph->E);
+	CA_ASSUME(theGraph->V);
+
 	if (gp_IsArc(e))
      {
     	 int e2 = gp_GetAdjacentArc(theGraph, e, link);
@@ -670,6 +735,11 @@ void gp_AttachArc(graphP theGraph, int v, int e, int link, int newArc)
 
 void gp_DetachArc(graphP theGraph, int arc)
 {
+	// Avoid false positive static analysis C6011 warning. Remove when it is fixed
+	CA_ASSUME(theGraph);
+	CA_ASSUME(theGraph->E);
+	CA_ASSUME(theGraph->V);
+
 	int nextArc = gp_GetNextArc(theGraph, arc),
 	    prevArc = gp_GetPrevArc(theGraph, arc);
 
@@ -703,7 +773,12 @@ void gp_DetachArc(graphP theGraph, int arc)
 
 int  gp_AddEdge(graphP theGraph, int u, int ulink, int v, int vlink)
 {
-int  upos, vpos;
+	// Avoid false positive static analysis C6011 warning. Remove when it is fixed
+	CA_ASSUME(theGraph);
+	CA_ASSUME(theGraph->E);
+	CA_ASSUME(theGraph->V);
+	
+	int  upos, vpos;
 
      if (theGraph==NULL || u < gp_GetFirstVertex(theGraph) || v < gp_GetFirstVertex(theGraph) ||
     		 !gp_VirtualVertexInRange(theGraph, u) || !gp_VirtualVertexInRange(theGraph, v))
@@ -745,7 +820,12 @@ int  upos, vpos;
 
 void _RestoreArc(graphP theGraph, int arc)
 {
-int nextArc = gp_GetNextArc(theGraph, arc),
+	// Avoid false positive static analysis C6011 warning. Remove when it is fixed
+	CA_ASSUME(theGraph);
+	CA_ASSUME(theGraph->E);
+	CA_ASSUME(theGraph->V);
+	
+	int nextArc = gp_GetNextArc(theGraph, arc),
 	prevArc = gp_GetPrevArc(theGraph, arc);
 
 	if (gp_IsArc(nextArc))
@@ -890,6 +970,11 @@ int  _HideVertex(graphP theGraph, int vertex)
 
 int gp_ContractEdge(graphP theGraph, int e)
 {
+	// Avoid false positive static analysis C6011 warning. Remove when it is fixed
+	CA_ASSUME(theGraph);
+	CA_ASSUME(theGraph->E);
+	CA_ASSUME(theGraph->V);
+
 	int eBefore, u, v;
 
 	if (gp_IsNotArc(e))
@@ -942,6 +1027,9 @@ int gp_ContractEdge(graphP theGraph, int e)
 
 int gp_IdentifyVertices(graphP theGraph, int u, int v, int eBefore)
 {
+	// Avoid false positive static analysis C6011 warning. Remove when it is fixed
+	CA_ASSUME(theGraph);
+
 	return theGraph->functions.fpIdentifyVertices(theGraph, u, v, eBefore);
 }
 
@@ -1129,7 +1217,11 @@ int _RestoreVertex(graphP theGraph)
 {
 int u, v, e_u_succ, e_u_pred, e_v_first, e_v_last, HESB, e;
 
-    if (sp_GetCurrentSize(theGraph->theStack) < 7)
+CA_ASSUME(theGraph);
+CA_ASSUME(theGraph->E);
+CA_ASSUME(theGraph->V);
+
+if (sp_GetCurrentSize(theGraph->theStack) < 7)
     	return NOTOK;
 
     sp_Pop(theGraph->theStack, v);
