Fix static analysis warnings.

diff --git a/Engine/Source/Developer/GraphColor/Private/graphUtils.cpp b/Engine/Source/Developer/GraphColor/Private/graphUtils.cpp
--- a/Engine/Source/Developer/GraphColor/Private/graphUtils.cpp
+++ b/Engine/Source/Developer/GraphColor/Private/graphUtils.cpp
@@ -703,6 +703,7 @@ void gp_AttachArc(graphP theGraph, int v, int e, int link, int newArc)
     	 gp_SetAdjacentArc(theGraph, newArc, 1^link, NIL);
 
     	 // newArcs's elink is e2
+		 CA_ASSUME(theGraph->E);
     	 gp_SetAdjacentArc(theGraph, newArc, link, e2);
 
     	 // if e2 is an arc, then e2's 1^link is newArc, else v's 1^link is newArc
@@ -1160,6 +1161,7 @@ int _IdentifyVertices(graphP theGraph, int u, int v, int eBefore)
         }
 
         gp_SetFirstArc(theGraph, v, NIL);
+		CA_ASSUME(theGraph->V);
         gp_SetLastArc(theGraph, v, NIL);
     }
 
