Static Analysis fixes

diff --git a/Engine/Source/Developer/GraphColor/Private/graphUtils.cpp b/Engine/Source/Developer/GraphColor/Private/graphUtils.cpp
--- a/Engine/Source/Developer/GraphColor/Private/graphUtils.cpp
+++ b/Engine/Source/Developer/GraphColor/Private/graphUtils.cpp
@@ -133,14 +133,11 @@ graphP theGraph = (graphP) malloc(sizeof(baseGraphStructure));
 
 void _InitFunctionTable(graphP theGraph)
 {
-	// Avoid false positive static analysis C6011 warning. Remove when it is fixed
-	CA_ASSUME(theGraph);
-
 	theGraph->functions.fpInitGraph = _InitGraph;
 
-     theGraph->functions.fpHideEdge = _HideEdge;
-     theGraph->functions.fpRestoreVertex = _RestoreVertex;
-     theGraph->functions.fpIdentifyVertices = _IdentifyVertices;
+	theGraph->functions.fpHideEdge = _HideEdge;
+	theGraph->functions.fpRestoreVertex = _RestoreVertex;
+	theGraph->functions.fpIdentifyVertices = _IdentifyVertices;
 }
 
 /********************************************************************
@@ -183,9 +180,6 @@ int gp_InitGraph(graphP theGraph, int N)
 
 int  _InitGraph(graphP theGraph, int N)
 {
-	// Avoid false positive static analysis C6011 warning. Remove when it is fixed
-	CA_ASSUME(theGraph);
-
 	int  Vsize, VIsize, Esize, stackSize;
 
 	 // Compute the vertex and edge capacities of the graph
@@ -223,11 +217,6 @@ int  _InitGraph(graphP theGraph, int N)
  ********************************************************************/
 void _InitVertices(graphP theGraph)
 {
-	// Avoid false positive static analysis C6011 warning. Remove when it is fixed
-	CA_ASSUME(theGraph);
-	CA_ASSUME(theGraph->E);
-	CA_ASSUME(theGraph->V);
-
 #if NIL == 0
 	memset(theGraph->V, NIL_CHAR, gp_VertexIndexBound(theGraph) * sizeof(vertexRec));
 #elif NIL == -1
@@ -247,11 +236,6 @@ void _InitVertices(graphP theGraph)
  ********************************************************************/
 void _InitEdges(graphP theGraph)
 {
-	// Avoid false positive static analysis C6011 warning. Remove when it is fixed
-	CA_ASSUME(theGraph);
-	CA_ASSUME(theGraph->E);
-	CA_ASSUME(theGraph->V);
-
 #if NIL == 0
 	memset(theGraph->E, NIL_CHAR, gp_EdgeIndexBound(theGraph) * sizeof(edgeRec));
 #elif NIL == -1
@@ -274,11 +258,6 @@ void _InitEdges(graphP theGraph)
  ********************************************************************/
 int gp_GetArcCapacity(graphP theGraph)
 {
-	// Avoid false positive static analysis C6011 warning. Remove when it is fixed
-	CA_ASSUME(theGraph);
-	CA_ASSUME(theGraph->E);
-	CA_ASSUME(theGraph->V);
-
 	return theGraph->arcCapacity - gp_GetFirstEdge(theGraph);
 }
 
@@ -338,11 +317,6 @@ int gp_EnsureArcCapacity(graphP theGraph, int requiredArcCapacity)
 	if (theGraph == NULL || requiredArcCapacity <= 0)
 		return NOTOK;
 
-	// Avoid false positive static analysis C6011 warning. Remove when it is fixed
-	CA_ASSUME(theGraph);
-	CA_ASSUME(theGraph->E);
-	CA_ASSUME(theGraph->V);
-
 	// Train callers to only ask for an even number of arcs, since
 	// two are required per edge or directed edge.
 	if (requiredArcCapacity & 1)
@@ -369,11 +343,6 @@ int gp_EnsureArcCapacity(graphP theGraph, int requiredArcCapacity)
 
 void _ClearVisitedFlags(graphP theGraph)
 {
-	// Avoid false positive static analysis C6011 warning. Remove when it is fixed
-	CA_ASSUME(theGraph);
-	CA_ASSUME(theGraph->E);
-	CA_ASSUME(theGraph->V);
-
 	_ClearVertexVisitedFlags(theGraph, TRUE);
 }
 
@@ -383,11 +352,6 @@ void _ClearVisitedFlags(graphP theGraph)
 
 void _ClearVertexVisitedFlags(graphP theGraph, int includeVirtualVertices)
 {
-	// Avoid false positive static analysis C6011 warning. Remove when it is fixed
-	CA_ASSUME(theGraph);
-	CA_ASSUME(theGraph->E);
-	CA_ASSUME(theGraph->V);
-
 	int  v;
 
 	for (v = gp_GetFirstVertex(theGraph); gp_VertexInRange(theGraph, v); v++)
@@ -406,9 +370,6 @@ void _ClearVertexVisitedFlags(graphP theGraph, int includeVirtualVertices)
 
 void _ClearGraph(graphP theGraph)
 {
-	// Avoid false positive static analysis C6011 warning. Remove when it is fixed
-	CA_ASSUME(theGraph);
-
 	if (theGraph->V != NULL)
      {
           free(theGraph->V);
@@ -474,14 +435,6 @@ int  v, e, Esize;
          return NOTOK;
      }
 
-	 // Avoid false positive static analysis C6011 warning. Remove when it is fixed
-	 CA_ASSUME(srcGraph);
-	 CA_ASSUME(dstGraph);
-	 CA_ASSUME(srcGraph->V);
-	 CA_ASSUME(dstGraph->V);
-	 CA_ASSUME(srcGraph->N);
-	 CA_ASSUME(dstGraph->N);
-
      // Ensure dstGraph has the required arc capacity; this expands
      // dstGraph if needed, but does not contract.  An error is only
      // returned if the expansion fails.
@@ -553,11 +506,6 @@ int  v, e, Esize;
 
 int  gp_IsNeighbor(graphP theGraph, int u, int v)
 {
-	// Avoid false positive static analysis C6011 warning. Remove when it is fixed
-	CA_ASSUME(theGraph);
-	CA_ASSUME(theGraph->E);
-	CA_ASSUME(theGraph->V);
-
 	int  e = gp_GetFirstArc(theGraph, u);
 
      while (gp_IsArc(e))
@@ -589,11 +537,6 @@ int  gp_IsNeighbor(graphP theGraph, int u, int v)
 
 int  gp_GetNeighborEdgeRecord(graphP theGraph, int u, int v)
 {
-	// Avoid false positive static analysis C6011 warning. Remove when it is fixed
-	CA_ASSUME(theGraph);
-	CA_ASSUME(theGraph->E);
-	CA_ASSUME(theGraph->V);
-
 	int  e;
 
      if (gp_IsNotVertex(u) || gp_IsNotVertex(v))
@@ -628,11 +571,6 @@ int  gp_GetNeighborEdgeRecord(graphP theGraph, int u, int v)
 
 int  gp_GetVertexDegree(graphP theGraph, int v)
 {
-	// Avoid false positive static analysis C6011 warning. Remove when it is fixed
-	CA_ASSUME(theGraph);
-	CA_ASSUME(theGraph->E);
-	CA_ASSUME(theGraph->V);
-	
 	int  e, degree;
 
      if (theGraph==NULL || gp_IsNotVertex(v))
@@ -672,11 +610,6 @@ int  gp_GetVertexDegree(graphP theGraph, int v)
 
 void gp_AttachArc(graphP theGraph, int v, int e, int link, int newArc)
 {
-	// Avoid false positive static analysis C6011 warning. Remove when it is fixed
-	CA_ASSUME(theGraph);
-	CA_ASSUME(theGraph->E);
-	CA_ASSUME(theGraph->V);
-
 	if (gp_IsArc(e))
      {
     	 int e2 = gp_GetAdjacentArc(theGraph, e, link);
@@ -703,7 +636,6 @@ void gp_AttachArc(graphP theGraph, int v, int e, int link, int newArc)
     	 gp_SetAdjacentArc(theGraph, newArc, 1^link, NIL);
 
     	 // newArcs's elink is e2
-		 CA_ASSUME(theGraph->E);
     	 gp_SetAdjacentArc(theGraph, newArc, link, e2);
 
     	 // if e2 is an arc, then e2's 1^link is newArc, else v's 1^link is newArc
@@ -736,11 +668,6 @@ void gp_AttachArc(graphP theGraph, int v, int e, int link, int newArc)
 
 void gp_DetachArc(graphP theGraph, int arc)
 {
-	// Avoid false positive static analysis C6011 warning. Remove when it is fixed
-	CA_ASSUME(theGraph);
-	CA_ASSUME(theGraph->E);
-	CA_ASSUME(theGraph->V);
-
 	int nextArc = gp_GetNextArc(theGraph, arc),
 	    prevArc = gp_GetPrevArc(theGraph, arc);
 
@@ -774,11 +701,6 @@ void gp_DetachArc(graphP theGraph, int arc)
 
 int  gp_AddEdge(graphP theGraph, int u, int ulink, int v, int vlink)
 {
-	// Avoid false positive static analysis C6011 warning. Remove when it is fixed
-	CA_ASSUME(theGraph);
-	CA_ASSUME(theGraph->E);
-	CA_ASSUME(theGraph->V);
-	
 	int  upos, vpos;
 
      if (theGraph==NULL || u < gp_GetFirstVertex(theGraph) || v < gp_GetFirstVertex(theGraph) ||
@@ -821,11 +743,6 @@ int  gp_AddEdge(graphP theGraph, int u, int ulink, int v, int vlink)
 
 void _RestoreArc(graphP theGraph, int arc)
 {
-	// Avoid false positive static analysis C6011 warning. Remove when it is fixed
-	CA_ASSUME(theGraph);
-	CA_ASSUME(theGraph->E);
-	CA_ASSUME(theGraph->V);
-	
 	int nextArc = gp_GetNextArc(theGraph, arc),
 	prevArc = gp_GetPrevArc(theGraph, arc);
 
@@ -971,11 +888,6 @@ int  _HideVertex(graphP theGraph, int vertex)
 
 int gp_ContractEdge(graphP theGraph, int e)
 {
-	// Avoid false positive static analysis C6011 warning. Remove when it is fixed
-	CA_ASSUME(theGraph);
-	CA_ASSUME(theGraph->E);
-	CA_ASSUME(theGraph->V);
-
 	int eBefore, u, v;
 
 	if (gp_IsNotArc(e))
@@ -1028,9 +940,6 @@ int gp_ContractEdge(graphP theGraph, int e)
 
 int gp_IdentifyVertices(graphP theGraph, int u, int v, int eBefore)
 {
-	// Avoid false positive static analysis C6011 warning. Remove when it is fixed
-	CA_ASSUME(theGraph);
-
 	return theGraph->functions.fpIdentifyVertices(theGraph, u, v, eBefore);
 }
 
@@ -1161,7 +1070,6 @@ int _IdentifyVertices(graphP theGraph, int u, int v, int eBefore)
         }
 
         gp_SetFirstArc(theGraph, v, NIL);
-		CA_ASSUME(theGraph->V);
         gp_SetLastArc(theGraph, v, NIL);
     }
 
@@ -1217,13 +1125,9 @@ int gp_RestoreVertex(graphP theGraph)
 
 int _RestoreVertex(graphP theGraph)
 {
-int u, v, e_u_succ, e_u_pred, e_v_first, e_v_last, HESB, e;
-
-CA_ASSUME(theGraph);
-CA_ASSUME(theGraph->E);
-CA_ASSUME(theGraph->V);
+	int u, v, e_u_succ, e_u_pred, e_v_first, e_v_last, HESB, e;
 
-if (sp_GetCurrentSize(theGraph->theStack) < 7)
+	if (sp_GetCurrentSize(theGraph->theStack) < 7)
     	return NOTOK;
 
     sp_Pop(theGraph->theStack, v);
diff --git a/Engine/Source/Runtime/Core/Public/Microsoft/MicrosoftPlatformCodeAnalysis.h b/Engine/Source/Runtime/Core/Public/Microsoft/MicrosoftPlatformCodeAnalysis.h
--- a/Engine/Source/Runtime/Core/Public/Microsoft/MicrosoftPlatformCodeAnalysis.h
+++ b/Engine/Source/Runtime/Core/Public/Microsoft/MicrosoftPlatformCodeAnalysis.h
@@ -110,7 +110,9 @@
 	//-V::505,542,581,601,623,668,677,690,704,719,720,730,735,751,758,780,781,1002,1008,1055,1062,1100
 	//-V:TRYCOMPRESSION:519,547
 	//-V:check(:501,547,560,605
+	//-V:checkSlow(:547
 	//-V:checkf(:510
+	//-V:checkfSlow(:547
 	//-V:dtAssert(:568
 	//-V:rcAssert(:568
 	//-V:GET_FUNCTION_NAME_CHECKED:521
@@ -313,4 +315,4 @@ __declspec(dllimport, noreturn) void CA_AssumeNoReturn();
 #endif
 
 
-#endif
+#endif
\ No newline at end of file
diff --git a/Engine/Source/Runtime/Engine/Classes/Sound/QuartzQuantizationUtilities.cpp b/Engine/Source/Runtime/Engine/Classes/Sound/QuartzQuantizationUtilities.cpp
--- a/Engine/Source/Runtime/Engine/Classes/Sound/QuartzQuantizationUtilities.cpp
+++ b/Engine/Source/Runtime/Engine/Classes/Sound/QuartzQuantizationUtilities.cpp
@@ -650,7 +650,7 @@ namespace Audio
 		checkSlow(MixerDevice);
 		checkSlow(MixerDevice->IsAudioRenderingThread());
 
-		if (CommandPtr && MixerDevice && !OwningClockName.IsNone())
+		if (CommandPtr && !OwningClockName.IsNone())
 		{
 			UE_LOG(LogAudioQuartz, Verbose, TEXT("OnQueued() called for quantized event type: [%s]"), *CommandPtr->GetCommandName().ToString());
 			return MixerDevice->QuantizedEventClockManager.CancelCommandOnClock(OwningClockName, CommandPtr);
diff --git a/Engine/Source/Runtime/Engine/Private/LocalPlayer.cpp b/Engine/Source/Runtime/Engine/Private/LocalPlayer.cpp
--- a/Engine/Source/Runtime/Engine/Private/LocalPlayer.cpp
+++ b/Engine/Source/Runtime/Engine/Private/LocalPlayer.cpp
@@ -878,6 +878,7 @@ FSceneView* ULocalPlayer::CalcSceneView( class FSceneViewFamily* ViewFamily,
 		// Base overrides (post process volumes, etc)
 		if (PlayerController->PlayerCameraManager)
 		{
+			checkSlow(CameraAnimPPBlendWeights && CameraAnimPPBlendOrders && CameraAnimPPSettings);
 			PlayerController->PlayerCameraManager->GetCachedPostProcessBlends(CameraAnimPPSettings, CameraAnimPPBlendWeights, CameraAnimPPBlendOrders);
 
 			for (int32 PPIdx = 0; PPIdx < CameraAnimPPBlendWeights->Num(); ++PPIdx)
diff --git a/Engine/Source/Runtime/Experimental/Chaos/Public/Chaos/TriangleMeshImplicitObject.h b/Engine/Source/Runtime/Experimental/Chaos/Public/Chaos/TriangleMeshImplicitObject.h
--- a/Engine/Source/Runtime/Experimental/Chaos/Public/Chaos/TriangleMeshImplicitObject.h
+++ b/Engine/Source/Runtime/Experimental/Chaos/Public/Chaos/TriangleMeshImplicitObject.h
@@ -444,6 +444,7 @@ namespace Chaos
 			else
 			{
 				static_assert(sizeof(ExpectedType) == 0, "Unsupported index buffer type");
+				return SmallIdxBuffer; // Appease the static analyzer by having a pointless return value
 			}
 		}
 
diff --git a/Engine/Source/Runtime/TextureUtilitiesCommon/Private/TextureImportUtils.cpp b/Engine/Source/Runtime/TextureUtilitiesCommon/Private/TextureImportUtils.cpp
--- a/Engine/Source/Runtime/TextureUtilitiesCommon/Private/TextureImportUtils.cpp
+++ b/Engine/Source/Runtime/TextureUtilitiesCommon/Private/TextureImportUtils.cpp
@@ -189,17 +189,18 @@ namespace UE
 
 			static bool IsOpaque(const ColorDataType InColor)
 			{
-				if ( sizeof(ColorDataType) == 4 )
+				if constexpr (sizeof(ColorDataType) == 4)
 				{
 					return InColor >= 0xFF000000U;
 				}
-				else if ( sizeof(ColorDataType) == 8 )
+				else if constexpr (sizeof(ColorDataType) == 8)
 				{
 					return InColor >= 0xFFFF000000000000ULL;
 				}
 				else
 				{
-					check(false);
+					static_assert(sizeof(ColorDataType) == 0);
+					return false;
 				}
 			}
 
